- Adds LLVM compile-time suppression file to CitySample allowing for sanitizer warnings to suppress symbol or source file path.
- Adds LLVM runtime suppression file to the engine allowing for runtime tsan issues to be suppressed
- Defines default rule to ignore issues generated when called from libemreee and libtbb as they have generated issues in the past and we cannot fix them without modifying third-party source
- Increases the TSAN test timeout to 5 hours (since we are timing out at 2 at the moment)

#rnx

diff --git a/Engine/Build/Linux/SanitizerRuntimeSuppressions.ini b/Engine/Build/Linux/SanitizerRuntimeSuppressions.ini
new file mode 100644
--- /dev/null
+++ b/Engine/Build/Linux/SanitizerRuntimeSuppressions.ini
@@ -0,0 +1,6 @@
+# Refer to [REDACTED_PROFILE_URL]/sanitizers/wiki/ThreadSanitizerSuppressions for information on the syntax
+
+#ignore all issues coming from particular libraries
+called_from_lib:libembree3.so
+called_from_lib:libtbbmalloc.so
+called_from_lib:libtbb.so
\ No newline at end of file
